Moved Kismet ScriptDisassembler into separate module

#rnx

diff --git a/Engine/Source/Developer/ScriptDisassembler/Private/ScriptDisassembler.cpp b/Engine/Source/Developer/ScriptDisassembler/Private/ScriptDisassembler.cpp
new file mode 100644
--- /dev/null
+++ b/Engine/Source/Developer/ScriptDisassembler/Private/ScriptDisassembler.cpp
@@ -0,0 +1,1067 @@
+// Copyright 1998-2019 Epic Games, Inc. All Rights Reserved.
+
+/*=============================================================================
+	ScriptDisassembler.cpp: Disassembler for Kismet bytecode.
+=============================================================================*/
+
+#include "ScriptDisassembler.h"
+#include "UObject/Object.h"
+#include "UObject/Class.h"
+#include "UObject/UObjectHash.h"
+#include "UObject/UObjectIterator.h"
+#include "UObject/UnrealType.h"
+
+DEFINE_LOG_CATEGORY_STATIC(LogScriptDisassembler, Log, All);
+
+/////////////////////////////////////////////////////
+// FKismetBytecodeDisassembler
+
+// Construct a disassembler that will output to the specified archive.
+FKismetBytecodeDisassembler::FKismetBytecodeDisassembler(FOutputDevice& InAr)
+	: Ar(InAr)
+{
+	InitTables();
+}
+
+// Disassemble all of the script code in a single structure.
+void FKismetBytecodeDisassembler::DisassembleStructure(UFunction* Source)
+{
+	Script.Empty();
+	Script.Append(Source->Script);
+
+	int32 ScriptIndex = 0;
+	while (ScriptIndex < Script.Num())
+	{
+		Ar.Logf(TEXT("Label_0x%X:"), ScriptIndex);
+
+		AddIndent();
+		SerializeExpr(ScriptIndex);
+		DropIndent();
+	}
+}
+
+// Disassemble all functions in any classes that have matching names.
+void FKismetBytecodeDisassembler::DisassembleAllFunctionsInClasses(FOutputDevice& Ar, const FString& ClassnameSubstring)
+{
+	FKismetBytecodeDisassembler Disasm(Ar);
+		
+	for (TObjectIterator<UClass> ClassIter; ClassIter; ++ClassIter)
+	{
+		UClass* Class = *ClassIter;
+
+		FString ClassName = Class->GetName();
+		if (FCString::Strifind(*ClassName, *ClassnameSubstring))
+		{
+			Ar.Logf(TEXT("Processing class %s"), *ClassName);
+
+			for (TFieldIterator<UFunction> FunctionIter(Class, EFieldIteratorFlags::ExcludeSuper); FunctionIter; ++FunctionIter)
+			{
+				UFunction* Function = *FunctionIter;
+				FString FunctionName = Function->GetName();
+				Ar.Logf(TEXT("  Processing function %s (%d bytes)"), *FunctionName, Function->Script.Num());
+
+				Disasm.DisassembleStructure(Function);
+
+				Ar.Logf(TEXT(""));
+			}
+
+			Ar.Logf(TEXT(""));
+			Ar.Logf(TEXT("-----------"));
+			Ar.Logf(TEXT(""));
+		}
+
+	}
+}
+
+EExprToken FKismetBytecodeDisassembler::SerializeExpr(int32& ScriptIndex)
+{
+	AddIndent();
+
+	EExprToken Opcode = (EExprToken)Script[ScriptIndex];
+	ScriptIndex++;
+
+	ProcessCommon(ScriptIndex, Opcode);
+
+	DropIndent();
+
+	return Opcode;
+}
+
+int32 FKismetBytecodeDisassembler::ReadINT(int32& ScriptIndex)
+{
+	int32 Value = Script[ScriptIndex]; ++ScriptIndex;
+	Value = Value | ((int32)Script[ScriptIndex] << 8); ++ScriptIndex;
+	Value = Value | ((int32)Script[ScriptIndex] << 16); ++ScriptIndex;
+	Value = Value | ((int32)Script[ScriptIndex] << 24); ++ScriptIndex;
+
+	return Value;
+}
+
+uint64 FKismetBytecodeDisassembler::ReadQWORD(int32& ScriptIndex)
+{
+	uint64 Value = Script[ScriptIndex]; ++ScriptIndex;
+	Value = Value | ((uint64)Script[ScriptIndex] << 8); ++ScriptIndex;
+	Value = Value | ((uint64)Script[ScriptIndex] << 16); ++ScriptIndex;
+	Value = Value | ((uint64)Script[ScriptIndex] << 24); ++ScriptIndex;
+	Value = Value | ((uint64)Script[ScriptIndex] << 32); ++ScriptIndex;
+	Value = Value | ((uint64)Script[ScriptIndex] << 40); ++ScriptIndex;
+	Value = Value | ((uint64)Script[ScriptIndex] << 48); ++ScriptIndex;
+	Value = Value | ((uint64)Script[ScriptIndex] << 56); ++ScriptIndex;
+
+	return Value;
+}
+
+uint8 FKismetBytecodeDisassembler::ReadBYTE(int32& ScriptIndex)
+{
+	uint8 Value = Script[ScriptIndex]; ++ScriptIndex;
+
+	return Value;
+}
+
+FString FKismetBytecodeDisassembler::ReadName(int32& ScriptIndex)
+{
+	const FScriptName ConstValue = *(FScriptName*)(Script.GetData() + ScriptIndex);
+	ScriptIndex += sizeof(FScriptName);
+
+	return ScriptNameToName(ConstValue).ToString();
+}
+
+uint16 FKismetBytecodeDisassembler::ReadWORD(int32& ScriptIndex)
+{
+	uint16 Value = Script[ScriptIndex]; ++ScriptIndex;
+	Value = Value | ((uint16)Script[ScriptIndex] << 8); ++ScriptIndex;
+	return Value;
+}
+
+float FKismetBytecodeDisassembler::ReadFLOAT(int32& ScriptIndex)
+{
+	union { float f; int32 i; } Result;
+	Result.i = ReadINT(ScriptIndex);
+	return Result.f;
+}
+
+CodeSkipSizeType FKismetBytecodeDisassembler::ReadSkipCount(int32& ScriptIndex)
+{
+#if SCRIPT_LIMIT_BYTECODE_TO_64KB
+	return ReadWORD(ScriptIndex);
+#else
+	static_assert(sizeof(CodeSkipSizeType) == 4, "Update this code as size changed.");
+	return ReadINT(ScriptIndex);
+#endif
+}
+
+FString FKismetBytecodeDisassembler::ReadString(int32& ScriptIndex)
+{
+	const EExprToken Opcode = (EExprToken)Script[ScriptIndex++];
+
+	switch (Opcode)
+	{
+	case EX_StringConst:
+		return ReadString8(ScriptIndex);
+
+	case EX_UnicodeStringConst:
+		return ReadString16(ScriptIndex);
+
+	default:
+		checkf(false, TEXT("FKismetBytecodeDisassembler::ReadString - Unexpected opcode. Expected %d or %d, got %d"), (int)EX_StringConst, (int)EX_UnicodeStringConst, (int)Opcode);
+		break;
+	}
+
+	return FString();
+}
+
+FString FKismetBytecodeDisassembler::ReadString8(int32& ScriptIndex)
+{
+	FString Result;
+
+	do
+	{
+		Result += (ANSICHAR)ReadBYTE(ScriptIndex);
+	}
+	while (Script[ScriptIndex-1] != 0);
+
+	return Result;
+}
+
+FString FKismetBytecodeDisassembler::ReadString16(int32& ScriptIndex)
+{
+	FString Result;
+
+	do
+	{
+		Result += (TCHAR)ReadWORD(ScriptIndex);
+	}
+	while ((Script[ScriptIndex-1] != 0) || (Script[ScriptIndex-2] != 0));
+
+	// Inline combine any surrogate pairs in the data when loading into a UTF-32 string
+	StringConv::InlineCombineSurrogates(Result);
+
+	return Result;
+}
+
+void FKismetBytecodeDisassembler::ProcessCastByte(int32 CastType, int32& ScriptIndex)
+{
+	// Expression of cast
+	SerializeExpr(ScriptIndex);
+}
+
+void FKismetBytecodeDisassembler::ProcessCommon(int32& ScriptIndex, EExprToken Opcode)
+{
+	switch (Opcode)
+	{
+	case EX_PrimitiveCast:
+		{
+			// A type conversion.
+			uint8 ConversionType = ReadBYTE(ScriptIndex);
+			Ar.Logf(TEXT("%s $%X: PrimitiveCast of type %d"), *Indents, (int32)Opcode, ConversionType);
+			AddIndent();
+
+			Ar.Logf(TEXT("%s Argument:"), *Indents);
+			ProcessCastByte(ConversionType, ScriptIndex);
+
+			//@TODO:
+			//Ar.Logf(TEXT("%s Expression:"), *Indents);
+			//SerializeExpr( ScriptIndex );
+			break;
+		}
+	case EX_SetSet:
+		{
+ 			Ar.Logf(TEXT("%s $%X: set set"), *Indents, (int32)Opcode);
+			SerializeExpr(ScriptIndex);
+			ReadINT(ScriptIndex);
+ 			while( SerializeExpr(ScriptIndex) != EX_EndSet)
+ 			{
+ 				// Set contents
+ 			}
+ 			break;
+		}
+	case EX_EndSet:
+		{
+			Ar.Logf(TEXT("%s $%X: EX_EndSet"), *Indents, (int32)Opcode);
+			break;
+		}
+	case EX_SetConst:
+		{
+			UProperty* InnerProp = ReadPointer<UProperty>(ScriptIndex);
+			int32 Num = ReadINT(ScriptIndex);
+ 			Ar.Logf(TEXT("%s $%X: set set const - elements number: %d, inner property: %s"), *Indents, (int32)Opcode, Num, *GetNameSafe(InnerProp));
+ 			while (SerializeExpr(ScriptIndex) != EX_EndSetConst)
+ 			{
+ 				// Set contents
+ 			}
+ 			break;
+		}
+	case EX_EndSetConst:
+		{
+			Ar.Logf(TEXT("%s $%X: EX_EndSetConst"), *Indents, (int32)Opcode);
+			break;
+		}
+	case EX_SetMap:
+		{
+ 			Ar.Logf(TEXT("%s $%X: set map"), *Indents, (int32)Opcode);
+			SerializeExpr(ScriptIndex);
+ 			ReadINT(ScriptIndex);
+ 			while( SerializeExpr(ScriptIndex) != EX_EndMap)
+ 			{
+ 				// Map contents
+ 			}
+ 			break;
+		}
+	case EX_EndMap:
+		{
+			Ar.Logf(TEXT("%s $%X: EX_EndMap"), *Indents, (int32)Opcode);
+			break;
+		}
+	case EX_MapConst:
+		{
+			UProperty* KeyProp = ReadPointer<UProperty>(ScriptIndex);
+			UProperty* ValProp = ReadPointer<UProperty>(ScriptIndex);
+			int32 Num = ReadINT(ScriptIndex);
+ 			Ar.Logf(TEXT("%s $%X: set map const - elements number: %d, key property: %s, val property: %s"), *Indents, (int32)Opcode, Num, *GetNameSafe(KeyProp), *GetNameSafe(ValProp));
+ 			while (SerializeExpr(ScriptIndex) != EX_EndMapConst)
+ 			{
+ 				// Map contents
+ 			}
+ 			break;
+		}
+	case EX_EndMapConst:
+		{
+			Ar.Logf(TEXT("%s $%X: EX_EndMapConst"), *Indents, (int32)Opcode);
+			break;
+		}
+	case EX_ObjToInterfaceCast:
+		{
+			// A conversion from an object variable to a native interface variable.
+			// We use a different bytecode to avoid the branching each time we process a cast token
+
+			// the interface class to convert to
+			UClass* InterfaceClass = ReadPointer<UClass>(ScriptIndex);
+			Ar.Logf(TEXT("%s $%X: ObjToInterfaceCast to %s"), *Indents, (int32)Opcode, *InterfaceClass->GetName());
+
+			SerializeExpr( ScriptIndex );
+			break;
+		}
+	case EX_CrossInterfaceCast:
+		{
+			// A conversion from one interface variable to a different interface variable.
+			// We use a different bytecode to avoid the branching each time we process a cast token
+
+			// the interface class to convert to
+			UClass* InterfaceClass = ReadPointer<UClass>(ScriptIndex);
+			Ar.Logf(TEXT("%s $%X: InterfaceToInterfaceCast to %s"), *Indents, (int32)Opcode, *InterfaceClass->GetName());
+
+			SerializeExpr( ScriptIndex );
+			break;
+		}
+	case EX_InterfaceToObjCast:
+		{
+			// A conversion from an interface variable to a object variable.
+			// We use a different bytecode to avoid the branching each time we process a cast token
+
+			// the interface class to convert to
+			UClass* ObjectClass = ReadPointer<UClass>(ScriptIndex);
+			Ar.Logf(TEXT("%s $%X: InterfaceToObjCast to %s"), *Indents, (int32)Opcode, *ObjectClass->GetName());
+
+			SerializeExpr( ScriptIndex );
+			break;
+		}
+	case EX_Let:
+		{
+			Ar.Logf(TEXT("%s $%X: Let (Variable = Expression)"), *Indents, (int32)Opcode);
+			AddIndent();
+
+			ReadPointer<UProperty>(ScriptIndex);
+
+			// Variable expr.
+			Ar.Logf(TEXT("%s Variable:"), *Indents);
+			SerializeExpr( ScriptIndex );
+
+			// Assignment expr.
+			Ar.Logf(TEXT("%s Expression:"), *Indents);
+			SerializeExpr( ScriptIndex );
+
+			DropIndent();
+			break;
+		}
+	case EX_LetObj:
+	case EX_LetWeakObjPtr:
+		{
+			if( Opcode == EX_LetObj )
+			{
+				Ar.Logf(TEXT("%s $%X: Let Obj (Variable = Expression)"), *Indents, (int32)Opcode);
+			}
+			else
+			{
+				Ar.Logf(TEXT("%s $%X: Let WeakObjPtr (Variable = Expression)"), *Indents, (int32)Opcode);
+			}
+			AddIndent();
+
+			// Variable expr.
+			Ar.Logf(TEXT("%s Variable:"), *Indents);
+			SerializeExpr( ScriptIndex );
+
+			// Assignment expr.
+			Ar.Logf(TEXT("%s Expression:"), *Indents);
+			SerializeExpr( ScriptIndex );
+
+			DropIndent();
+			break;
+		}
+	case EX_LetBool:
+		{
+			Ar.Logf(TEXT("%s $%X: LetBool (Variable = Expression)"), *Indents, (int32)Opcode);
+			AddIndent();
+
+			// Variable expr.
+			Ar.Logf(TEXT("%s Variable:"), *Indents);
+			SerializeExpr( ScriptIndex );
+
+			// Assignment expr.
+			Ar.Logf(TEXT("%s Expression:"), *Indents);
+			SerializeExpr( ScriptIndex );
+
+			DropIndent();
+			break;
+		}
+	case EX_LetValueOnPersistentFrame:
+		{
+			Ar.Logf(TEXT("%s $%X: LetValueOnPersistentFrame"), *Indents, (int32)Opcode);
+			AddIndent();
+
+			auto Prop = ReadPointer<UProperty>(ScriptIndex);
+			Ar.Logf(TEXT("%s Destination variable: %s, offset: %d"), *Indents, *GetNameSafe(Prop), 
+				Prop ? Prop->GetOffset_ForDebug() : 0);
+			
+			Ar.Logf(TEXT("%s Expression:"), *Indents);
+			SerializeExpr(ScriptIndex);
+
+			DropIndent();
+
+			break;
+		}
+	case EX_StructMemberContext:
+		{
+			Ar.Logf(TEXT("%s $%X: Struct member context "), *Indents, (int32)Opcode);
+			AddIndent();
+
+			UProperty* Prop = ReadPointer<UProperty>(ScriptIndex);
+
+			Ar.Logf(TEXT("%s Expression within struct %s, offset %d"), *Indents, *(Prop->GetName()), 
+				Prop->GetOffset_ForDebug()); // although that isn't a UFunction, we are not going to indirect the props of a struct, so this should be fine
+
+			Ar.Logf(TEXT("%s Expression to struct:"), *Indents);
+			SerializeExpr( ScriptIndex );
+
+			DropIndent();
+
+			break;
+		}
+	case EX_LetDelegate:
+		{
+			Ar.Logf(TEXT("%s $%X: LetDelegate (Variable = Expression)"), *Indents, (int32)Opcode);
+			AddIndent();
+
+			// Variable expr.
+			Ar.Logf(TEXT("%s Variable:"), *Indents);
+			SerializeExpr( ScriptIndex );
+				
+			// Assignment expr.
+			Ar.Logf(TEXT("%s Expression:"), *Indents);
+			SerializeExpr( ScriptIndex );
+
+			DropIndent();
+			break;
+		}
+	case EX_LocalVirtualFunction:
+		{
+			FString FunctionName = ReadName(ScriptIndex);
+			Ar.Logf(TEXT("%s $%X: Local Virtual Script Function named %s"), *Indents, (int32)Opcode, *FunctionName);
+
+			while (SerializeExpr(ScriptIndex) != EX_EndFunctionParms)
+			{
+			}
+			break;
+		}
+	case EX_LocalFinalFunction:
+		{
+			UStruct* StackNode = ReadPointer<UStruct>(ScriptIndex);
+			Ar.Logf(TEXT("%s $%X: Local Final Script Function (stack node %s::%s)"), *Indents, (int32)Opcode, StackNode ? *StackNode->GetOuter()->GetName() : TEXT("(null)"), StackNode ? *StackNode->GetName() : TEXT("(null)"));
+
+			while (SerializeExpr( ScriptIndex ) != EX_EndFunctionParms)
+			{
+				// Params
+			}
+			break;
+		}
+	case EX_LetMulticastDelegate:
+		{
+			Ar.Logf(TEXT("%s $%X: LetMulticastDelegate (Variable = Expression)"), *Indents, (int32)Opcode);
+			AddIndent();
+
+			// Variable expr.
+			Ar.Logf(TEXT("%s Variable:"), *Indents);
+			SerializeExpr( ScriptIndex );
+				
+			// Assignment expr.
+			Ar.Logf(TEXT("%s Expression:"), *Indents);
+			SerializeExpr( ScriptIndex );
+
+			DropIndent();
+			break;
+		}
+
+	case EX_ComputedJump:
+		{
+			Ar.Logf(TEXT("%s $%X: Computed Jump, offset specified by expression:"), *Indents, (int32)Opcode);
+
+			AddIndent();
+			SerializeExpr( ScriptIndex );
+			DropIndent();
+
+			break;
+		}
+
+	case EX_Jump:
+		{
+			CodeSkipSizeType SkipCount = ReadSkipCount(ScriptIndex);
+			Ar.Logf(TEXT("%s $%X: Jump to offset 0x%X"), *Indents, (int32)Opcode, SkipCount);
+			break;
+		}
+	case EX_LocalVariable:
+		{
+			UProperty* PropertyPtr = ReadPointer<UProperty>(ScriptIndex);
+			Ar.Logf(TEXT("%s $%X: Local variable named %s"), *Indents, (int32)Opcode, PropertyPtr ? *PropertyPtr->GetName() : TEXT("(null)"));
+			break;
+		}
+	case EX_DefaultVariable:
+		{
+			UProperty* PropertyPtr = ReadPointer<UProperty>(ScriptIndex);
+			Ar.Logf(TEXT("%s $%X: Default variable named %s"), *Indents, (int32)Opcode, PropertyPtr ? *PropertyPtr->GetName() : TEXT("(null)"));
+			break;
+		}
+	case EX_InstanceVariable:
+		{
+			UProperty* PropertyPtr = ReadPointer<UProperty>(ScriptIndex);
+			Ar.Logf(TEXT("%s $%X: Instance variable named %s"), *Indents, (int32)Opcode, PropertyPtr ? *PropertyPtr->GetName() : TEXT("(null)"));
+			break;
+		}
+	case EX_LocalOutVariable:
+		{
+			UProperty* PropertyPtr = ReadPointer<UProperty>(ScriptIndex);
+			Ar.Logf(TEXT("%s $%X: Local out variable named %s"), *Indents, (int32)Opcode, PropertyPtr ? *PropertyPtr->GetName() : TEXT("(null)"));
+			break;
+		}
+	case EX_InterfaceContext:
+		{
+			Ar.Logf(TEXT("%s $%X: EX_InterfaceContext:"), *Indents, (int32)Opcode);
+			SerializeExpr(ScriptIndex);
+			break;
+		}
+	case EX_DeprecatedOp4A:
+		{
+			Ar.Logf(TEXT("%s $%X: This opcode has been removed and does nothing."), *Indents, (int32)Opcode);
+			break;
+		}
+	case EX_Nothing:
+		{
+			Ar.Logf(TEXT("%s $%X: EX_Nothing"), *Indents, (int32)Opcode);
+			break;
+		}
+	case EX_EndOfScript:
+		{
+			Ar.Logf(TEXT("%s $%X: EX_EndOfScript"), *Indents, (int32)Opcode);
+			break;
+		}
+	case EX_EndFunctionParms:
+		{
+			Ar.Logf(TEXT("%s $%X: EX_EndFunctionParms"), *Indents, (int32)Opcode);
+			break;
+		}
+	case EX_EndStructConst:
+		{
+			Ar.Logf(TEXT("%s $%X: EX_EndStructConst"), *Indents, (int32)Opcode);
+			break;
+		}
+	case EX_EndArray:
+		{
+			Ar.Logf(TEXT("%s $%X: EX_EndArray"), *Indents, (int32)Opcode);
+			break;
+		}
+	case EX_EndArrayConst:
+		{
+			Ar.Logf(TEXT("%s $%X: EX_EndArrayConst"), *Indents, (int32)Opcode);
+			break;
+		}
+	case EX_IntZero:
+		{
+			Ar.Logf(TEXT("%s $%X: EX_IntZero"), *Indents, (int32)Opcode);
+			break;
+		}
+	case EX_IntOne:
+		{
+			Ar.Logf(TEXT("%s $%X: EX_IntOne"), *Indents, (int32)Opcode);
+			break;
+		}
+	case EX_True:
+		{
+			Ar.Logf(TEXT("%s $%X: EX_True"), *Indents, (int32)Opcode);
+			break;
+		}
+	case EX_False:
+		{
+			Ar.Logf(TEXT("%s $%X: EX_False"), *Indents, (int32)Opcode);
+			break;
+		}
+	case EX_NoObject:
+		{
+			Ar.Logf(TEXT("%s $%X: EX_NoObject"), *Indents, (int32)Opcode);
+			break;
+		}
+	case EX_NoInterface:
+		{
+			Ar.Logf(TEXT("%s $%X: EX_NoObject"), *Indents, (int32)Opcode);
+			break;
+		}
+	case EX_Self:
+		{
+			Ar.Logf(TEXT("%s $%X: EX_Self"), *Indents, (int32)Opcode);
+			break;
+		}
+	case EX_EndParmValue:
+		{
+			Ar.Logf(TEXT("%s $%X: EX_EndParmValue"), *Indents, (int32)Opcode);
+			break;
+		}
+	case EX_Return:
+		{
+			Ar.Logf(TEXT("%s $%X: Return expression"), *Indents, (int32)Opcode);
+
+			SerializeExpr( ScriptIndex ); // Return expression.
+			break;
+		}
+	case EX_CallMath:
+		{
+			UStruct* StackNode = ReadPointer<UStruct>(ScriptIndex);
+			Ar.Logf(TEXT("%s $%X: Call Math (stack node %s::%s)"), *Indents, (int32)Opcode, *GetNameSafe(StackNode ? StackNode->GetOuter() : nullptr), *GetNameSafe(StackNode));
+
+			while (SerializeExpr(ScriptIndex) != EX_EndFunctionParms)
+			{
+				// Params
+			}
+			break;
+		}
+	case EX_FinalFunction:
+		{
+			UStruct* StackNode = ReadPointer<UStruct>(ScriptIndex);
+			Ar.Logf(TEXT("%s $%X: Final Function (stack node %s::%s)"), *Indents, (int32)Opcode, StackNode ? *StackNode->GetOuter()->GetName() : TEXT("(null)"), StackNode ? *StackNode->GetName() : TEXT("(null)"));
+
+			while (SerializeExpr( ScriptIndex ) != EX_EndFunctionParms)
+			{
+				// Params
+			}
+			break;
+		}
+	case EX_CallMulticastDelegate:
+		{
+			UStruct* StackNode = ReadPointer<UStruct>(ScriptIndex);
+			Ar.Logf(TEXT("%s $%X: CallMulticastDelegate (signature %s::%s) delegate:"), *Indents, (int32)Opcode, StackNode ? *StackNode->GetOuter()->GetName() : TEXT("(null)"), StackNode ? *StackNode->GetName() : TEXT("(null)"));
+			SerializeExpr( ScriptIndex );
+			Ar.Logf(TEXT("Params:"));
+			while (SerializeExpr( ScriptIndex ) != EX_EndFunctionParms)
+			{
+				// Params
+			}
+			break;
+		}
+	case EX_VirtualFunction:
+		{
+			FString FunctionName = ReadName(ScriptIndex);
+			Ar.Logf(TEXT("%s $%X: Virtual Function named %s"), *Indents, (int32)Opcode, *FunctionName);
+
+			while (SerializeExpr(ScriptIndex) != EX_EndFunctionParms)
+			{
+			}
+			break;
+		}
+	case EX_ClassContext:
+	case EX_Context:
+	case EX_Context_FailSilent:
+		{
+			Ar.Logf(TEXT("%s $%X: %s"), *Indents, (int32)Opcode, Opcode == EX_ClassContext ? TEXT("Class Context") : TEXT("Context"));
+			AddIndent();
+
+			// Object expression.
+			Ar.Logf(TEXT("%s ObjectExpression:"), *Indents);
+			SerializeExpr( ScriptIndex );
+
+			if (Opcode == EX_Context_FailSilent)
+			{
+				Ar.Logf(TEXT(" Can fail silently on access none "));
+			}
+
+			// Code offset for NULL expressions.
+			CodeSkipSizeType SkipCount = ReadSkipCount(ScriptIndex);
+			Ar.Logf(TEXT("%s Skip Bytes: 0x%X"), *Indents, SkipCount);
+
+			// Property corresponding to the r-value data, in case the l-value needs to be mem-zero'd
+			UField* Field = ReadPointer<UField>(ScriptIndex);
+			Ar.Logf(TEXT("%s R-Value Property: %s"), *Indents, Field ? *Field->GetName() : TEXT("(null)"));
+
+			// Context expression.
+			Ar.Logf(TEXT("%s ContextExpression:"), *Indents);
+			SerializeExpr( ScriptIndex );
+
+			DropIndent();
+			break;
+		}
+	case EX_IntConst:
+		{
+			int32 ConstValue = ReadINT(ScriptIndex);
+			Ar.Logf(TEXT("%s $%X: literal int32 %d"), *Indents, (int32)Opcode, ConstValue);
+			break;
+		}
+	case EX_SkipOffsetConst:
+		{
+			CodeSkipSizeType ConstValue = ReadSkipCount(ScriptIndex);
+			Ar.Logf(TEXT("%s $%X: literal CodeSkipSizeType 0x%X"), *Indents, (int32)Opcode, ConstValue);
+			break;
+		}
+	case EX_FloatConst:
+		{
+			float ConstValue = ReadFLOAT(ScriptIndex);
+			Ar.Logf(TEXT("%s $%X: literal float %f"), *Indents, (int32)Opcode, ConstValue);
+			break;
+		}
+	case EX_StringConst:
+		{
+			FString ConstValue = ReadString8(ScriptIndex);
+			Ar.Logf(TEXT("%s $%X: literal ansi string \"%s\""), *Indents, (int32)Opcode, *ConstValue);
+			break;
+		}
+	case EX_UnicodeStringConst:
+		{
+			FString ConstValue = ReadString16(ScriptIndex);
+			Ar.Logf(TEXT("%s $%X: literal unicode string \"%s\""), *Indents, (int32)Opcode, *ConstValue);
+			break;
+		}
+	case EX_TextConst:
+		{
+			// What kind of text are we dealing with?
+			const EBlueprintTextLiteralType TextLiteralType = (EBlueprintTextLiteralType)Script[ScriptIndex++];
+
+			switch (TextLiteralType)
+			{
+			case EBlueprintTextLiteralType::Empty:
+				{
+					Ar.Logf(TEXT("%s $%X: literal text - empty"), *Indents, (int32)Opcode);
+				}
+				break;
+
+			case EBlueprintTextLiteralType::LocalizedText:
+				{
+					const FString SourceString = ReadString(ScriptIndex);
+					const FString KeyString = ReadString(ScriptIndex);
+					const FString Namespace = ReadString(ScriptIndex);
+					Ar.Logf(TEXT("%s $%X: literal text - localized text { namespace: \"%s\", key: \"%s\", source: \"%s\" }"), *Indents, (int32)Opcode, *Namespace, *KeyString, *SourceString);
+				}
+				break;
+
+			case EBlueprintTextLiteralType::InvariantText:
+				{
+					const FString SourceString = ReadString(ScriptIndex);
+					Ar.Logf(TEXT("%s $%X: literal text - invariant text: \"%s\""), *Indents, (int32)Opcode, *SourceString);
+				}
+				break;
+
+			case EBlueprintTextLiteralType::LiteralString:
+				{
+					const FString SourceString = ReadString(ScriptIndex);
+					Ar.Logf(TEXT("%s $%X: literal text - literal string: \"%s\""), *Indents, (int32)Opcode, *SourceString);
+				}
+				break;
+
+			case EBlueprintTextLiteralType::StringTableEntry:
+				{
+					ReadPointer<UObject>(ScriptIndex); // String Table asset (if any)
+					const FString TableIdString = ReadString(ScriptIndex);
+					const FString KeyString = ReadString(ScriptIndex);
+					Ar.Logf(TEXT("%s $%X: literal text - string table entry { tableid: \"%s\", key: \"%s\" }"), *Indents, (int32)Opcode, *TableIdString, *KeyString);
+				}
+				break;
+
+			default:
+				checkf(false, TEXT("Unknown EBlueprintTextLiteralType! Please update FKismetBytecodeDisassembler::ProcessCommon to handle this type of text."));
+				break;
+			}
+			break;
+		}
+	case EX_ObjectConst:
+		{
+			UObject* Pointer = ReadPointer<UObject>(ScriptIndex);
+			Ar.Logf(TEXT("%s $%X: EX_ObjectConst (%p:%s)"), *Indents, (int32)Opcode, Pointer, *Pointer->GetFullName());
+			break;
+		}
+	case EX_SoftObjectConst:
+		{
+			Ar.Logf(TEXT("%s $%X: EX_SoftObjectConst"), *Indents, (int32)Opcode);
+			SerializeExpr(ScriptIndex);
+			break;
+		}
+	case EX_NameConst:
+		{
+			FString ConstValue = ReadName(ScriptIndex);
+			Ar.Logf(TEXT("%s $%X: literal name %s"), *Indents, (int32)Opcode, *ConstValue);
+			break;
+		}
+	case EX_RotationConst:
+		{
+			float Pitch = ReadFLOAT(ScriptIndex);
+			float Yaw = ReadFLOAT(ScriptIndex);
+			float Roll = ReadFLOAT(ScriptIndex);
+
+			Ar.Logf(TEXT("%s $%X: literal rotation (%f,%f,%f)"), *Indents, (int32)Opcode, Pitch, Yaw, Roll);
+			break;
+		}
+	case EX_VectorConst:
+		{
+			float X = ReadFLOAT(ScriptIndex);
+			float Y = ReadFLOAT(ScriptIndex);
+			float Z = ReadFLOAT(ScriptIndex);
+
+			Ar.Logf(TEXT("%s $%X: literal vector (%f,%f,%f)"), *Indents, (int32)Opcode, X, Y, Z);
+			break;
+		}
+	case EX_TransformConst:
+		{
+
+			float RotX = ReadFLOAT(ScriptIndex);
+			float RotY = ReadFLOAT(ScriptIndex);
+			float RotZ = ReadFLOAT(ScriptIndex);
+			float RotW = ReadFLOAT(ScriptIndex);
+
+			float TransX = ReadFLOAT(ScriptIndex);
+			float TransY = ReadFLOAT(ScriptIndex);
+			float TransZ = ReadFLOAT(ScriptIndex);
+
+			float ScaleX = ReadFLOAT(ScriptIndex);
+			float ScaleY = ReadFLOAT(ScriptIndex);
+			float ScaleZ = ReadFLOAT(ScriptIndex);
+
+			Ar.Logf(TEXT("%s $%X: literal transform R(%f,%f,%f,%f) T(%f,%f,%f) S(%f,%f,%f)"), *Indents, (int32)Opcode, TransX, TransY, TransZ, RotX, RotY, RotZ, RotW, ScaleX, ScaleY, ScaleZ);
+			break;
+		}
+	case EX_StructConst:
+		{
+			UScriptStruct* Struct = ReadPointer<UScriptStruct>(ScriptIndex);
+			int32 SerializedSize = ReadINT(ScriptIndex);
+			Ar.Logf(TEXT("%s $%X: literal struct %s (serialized size: %d)"), *Indents, (int32)Opcode, *Struct->GetName(), SerializedSize);
+			while( SerializeExpr(ScriptIndex) != EX_EndStructConst )
+			{
+				// struct contents
+			}
+			break;
+		}
+	case EX_SetArray:
+		{
+ 			Ar.Logf(TEXT("%s $%X: set array"), *Indents, (int32)Opcode);
+			SerializeExpr(ScriptIndex);
+ 			while( SerializeExpr(ScriptIndex) != EX_EndArray)
+ 			{
+ 				// Array contents
+ 			}
+ 			break;
+		}
+	case EX_ArrayConst:
+		{
+			UProperty* InnerProp = ReadPointer<UProperty>(ScriptIndex);
+			int32 Num = ReadINT(ScriptIndex);
+			Ar.Logf(TEXT("%s $%X: set array const - elements number: %d, inner property: %s"), *Indents, (int32)Opcode, Num, *GetNameSafe(InnerProp));
+			while (SerializeExpr(ScriptIndex) != EX_EndArrayConst)
+			{
+				// Array contents
+			}
+			break;
+		}
+	case EX_ByteConst:
+		{
+			uint8 ConstValue = ReadBYTE(ScriptIndex);
+			Ar.Logf(TEXT("%s $%X: literal byte %d"), *Indents, (int32)Opcode, ConstValue);
+			break;
+		}
+	case EX_IntConstByte:
+		{
+			int32 ConstValue = ReadBYTE(ScriptIndex);
+			Ar.Logf(TEXT("%s $%X: literal int %d"), *Indents, (int32)Opcode, ConstValue);
+			break;
+		}
+	case EX_MetaCast:
+		{
+			UClass* Class = ReadPointer<UClass>(ScriptIndex);
+			Ar.Logf(TEXT("%s $%X: MetaCast to %s of expr:"), *Indents, (int32)Opcode, *Class->GetName());
+			SerializeExpr( ScriptIndex );
+			break;
+		}
+	case EX_DynamicCast:
+		{
+			UClass* Class = ReadPointer<UClass>(ScriptIndex);
+			Ar.Logf(TEXT("%s $%X: DynamicCast to %s of expr:"), *Indents, (int32)Opcode, *Class->GetName());
+			SerializeExpr( ScriptIndex );
+			break;
+		}
+	case EX_JumpIfNot:
+		{
+			// Code offset.
+			CodeSkipSizeType SkipCount = ReadSkipCount(ScriptIndex);
+				
+			Ar.Logf(TEXT("%s $%X: Jump to offset 0x%X if not expr:"), *Indents, (int32)Opcode, SkipCount);
+
+			// Boolean expr.
+			SerializeExpr( ScriptIndex );
+			break;
+		}
+	case EX_Assert:
+		{
+			uint16 LineNumber = ReadWORD(ScriptIndex);
+			uint8 InDebugMode = ReadBYTE(ScriptIndex);
+
+			Ar.Logf(TEXT("%s $%X: assert at line %d, in debug mode = %d with expr:"), *Indents, (int32)Opcode, LineNumber, InDebugMode);
+			SerializeExpr( ScriptIndex ); // Assert expr.
+			break;
+		}
+	case EX_Skip:
+		{
+			CodeSkipSizeType W = ReadSkipCount(ScriptIndex);
+			Ar.Logf(TEXT("%s $%X: possibly skip 0x%X bytes of expr:"), *Indents, (int32)Opcode, W);
+
+			// Expression to possibly skip.
+			SerializeExpr( ScriptIndex );
+
+			break;
+		}
+	case EX_InstanceDelegate:
+		{
+			// the name of the function assigned to the delegate.
+			FString FuncName = ReadName(ScriptIndex);
+
+			Ar.Logf(TEXT("%s $%X: instance delegate function named %s"), *Indents, (int32)Opcode, *FuncName);
+			break;
+		}
+	case EX_AddMulticastDelegate:
+		{
+			Ar.Logf(TEXT("%s $%X: Add MC delegate"), *Indents, (int32)Opcode);
+			SerializeExpr( ScriptIndex );
+			SerializeExpr( ScriptIndex );
+			break;
+		}
+	case EX_RemoveMulticastDelegate:
+		{
+			Ar.Logf(TEXT("%s $%X: Remove MC delegate"), *Indents, (int32)Opcode);
+			SerializeExpr( ScriptIndex );
+			SerializeExpr( ScriptIndex );
+			break;
+		}
+	case EX_ClearMulticastDelegate:
+		{
+			Ar.Logf(TEXT("%s $%X: Clear MC delegate"), *Indents, (int32)Opcode);
+			SerializeExpr( ScriptIndex );
+			break;
+		}
+	case EX_BindDelegate:
+		{
+			// the name of the function assigned to the delegate.
+			FString FuncName = ReadName(ScriptIndex);
+
+			Ar.Logf(TEXT("%s $%X: BindDelegate '%s' "), *Indents, (int32)Opcode, *FuncName);
+
+			Ar.Logf(TEXT("%s Delegate:"), *Indents);
+			SerializeExpr( ScriptIndex );
+
+			Ar.Logf(TEXT("%s Object:"), *Indents);
+			SerializeExpr( ScriptIndex );
+
+			break;
+		}
+	case EX_PushExecutionFlow:
+		{
+			CodeSkipSizeType SkipCount = ReadSkipCount(ScriptIndex);
+			Ar.Logf(TEXT("%s $%X: FlowStack.Push(0x%X);"), *Indents, (int32)Opcode, SkipCount);
+			break;
+		}
+	case EX_PopExecutionFlow:
+		{
+			Ar.Logf(TEXT("%s $%X: if (FlowStack.Num()) { jump to statement at FlowStack.Pop(); } else { ERROR!!! }"), *Indents, (int32)Opcode);
+			break;
+		}
+	case EX_PopExecutionFlowIfNot:
+		{
+			Ar.Logf(TEXT("%s $%X: if (!condition) { if (FlowStack.Num()) { jump to statement at FlowStack.Pop(); } else { ERROR!!! } }"), *Indents, (int32)Opcode);
+			// Boolean expr.
+			SerializeExpr( ScriptIndex );
+			break;
+		}
+	case EX_Breakpoint:
+		{
+			Ar.Logf(TEXT("%s $%X: <<< BREAKPOINT >>>"), *Indents, (int32)Opcode);
+			break;
+		}
+	case EX_WireTracepoint:
+		{
+			Ar.Logf(TEXT("%s $%X: .. wire debug site .."), *Indents, (int32)Opcode);
+			break;
+		}
+	case EX_InstrumentationEvent:
+		{
+			const uint8 EventType = ReadBYTE(ScriptIndex);
+			switch (EventType)
+			{
+				case EScriptInstrumentation::InlineEvent:
+					Ar.Logf(TEXT("%s $%X: .. instrumented inline event .."), *Indents, (int32)Opcode);
+					break;
+				case EScriptInstrumentation::Stop:
+					Ar.Logf(TEXT("%s $%X: .. instrumented event stop .."), *Indents, (int32)Opcode);
+					break;
+				case EScriptInstrumentation::PureNodeEntry:
+					Ar.Logf(TEXT("%s $%X: .. instrumented pure node entry site .."), *Indents, (int32)Opcode);
+					break;
+				case EScriptInstrumentation::NodeDebugSite:
+					Ar.Logf(TEXT("%s $%X: .. instrumented debug site .."), *Indents, (int32)Opcode);
+					break;
+				case EScriptInstrumentation::NodeEntry:
+					Ar.Logf(TEXT("%s $%X: .. instrumented wire entry site .."), *Indents, (int32)Opcode);
+					break;
+				case EScriptInstrumentation::NodeExit:
+					Ar.Logf(TEXT("%s $%X: .. instrumented wire exit site .."), *Indents, (int32)Opcode);
+					break;
+				case EScriptInstrumentation::PushState:
+					Ar.Logf(TEXT("%s $%X: .. push execution state .."), *Indents, (int32)Opcode);
+					break;
+				case EScriptInstrumentation::RestoreState:
+					Ar.Logf(TEXT("%s $%X: .. restore execution state .."), *Indents, (int32)Opcode);
+					break;
+				case EScriptInstrumentation::ResetState:
+					Ar.Logf(TEXT("%s $%X: .. reset execution state .."), *Indents, (int32)Opcode);
+					break;
+				case EScriptInstrumentation::SuspendState:
+					Ar.Logf(TEXT("%s $%X: .. suspend execution state .."), *Indents, (int32)Opcode);
+					break;
+				case EScriptInstrumentation::PopState:
+					Ar.Logf(TEXT("%s $%X: .. pop execution state .."), *Indents, (int32)Opcode);
+					break;
+				case EScriptInstrumentation::TunnelEndOfThread:
+					Ar.Logf(TEXT("%s $%X: .. tunnel end of thread .."), *Indents, (int32)Opcode);
+					break;
+			}
+			break;
+		}
+	case EX_Tracepoint:
+		{
+			Ar.Logf(TEXT("%s $%X: .. debug site .."), *Indents, (int32)Opcode);
+			break;
+		}
+	case EX_SwitchValue:
+		{
+			const auto NumCases = ReadWORD(ScriptIndex);
+			const auto AfterSkip = ReadSkipCount(ScriptIndex);
+
+			Ar.Logf(TEXT("%s $%X: Switch Value %d cases, end in 0x%X"), *Indents, (int32)Opcode, NumCases, AfterSkip);
+			AddIndent();
+			Ar.Logf(TEXT("%s Index:"), *Indents);
+			SerializeExpr(ScriptIndex);
+
+			for (uint16 CaseIndex = 0; CaseIndex < NumCases; ++CaseIndex)
+			{
+				Ar.Logf(TEXT("%s [%d] Case Index (label: 0x%X):"), *Indents, CaseIndex, ScriptIndex);
+				SerializeExpr(ScriptIndex);	// case index value term
+				const auto OffsetToNextCase = ReadSkipCount(ScriptIndex);
+				Ar.Logf(TEXT("%s [%d] Offset to the next case: 0x%X"), *Indents, CaseIndex, OffsetToNextCase);
+				Ar.Logf(TEXT("%s [%d] Case Result:"), *Indents, CaseIndex);
+				SerializeExpr(ScriptIndex);	// case term
+			}
+
+			Ar.Logf(TEXT("%s Default result (label: 0x%X):"), *Indents, ScriptIndex);
+			SerializeExpr(ScriptIndex);
+			Ar.Logf(TEXT("%s (label: 0x%X)"), *Indents, ScriptIndex);
+			DropIndent();
+			break;
+		}
+	case EX_ArrayGetByRef:
+		{
+			Ar.Logf(TEXT("%s $%X: Array Get-by-Ref Index"), *Indents, (int32)Opcode);
+			AddIndent();
+			SerializeExpr(ScriptIndex);
+			SerializeExpr(ScriptIndex);
+			DropIndent();
+			break;
+		}
+	default:
+		{
+			// This should never occur.
+			UE_LOG(LogScriptDisassembler, Warning, TEXT("Unknown bytecode 0x%02X; ignoring it"), (uint8)Opcode );
+			break;
+		}
+	}
+}
+
+void FKismetBytecodeDisassembler::InitTables()
+{
+}
diff --git a/Engine/Source/Developer/ScriptDisassembler/Private/ScriptDisassemblerModule.cpp b/Engine/Source/Developer/ScriptDisassembler/Private/ScriptDisassemblerModule.cpp
new file mode 100644
--- /dev/null
+++ b/Engine/Source/Developer/ScriptDisassembler/Private/ScriptDisassemblerModule.cpp
@@ -0,0 +1,5 @@
+// Copyright 1998-2019 Epic Games, Inc. All Rights Reserved.
+
+#include "Modules/ModuleManager.h"
+
+IMPLEMENT_MODULE( FDefaultModuleImpl, ScriptDisassembler );
diff --git a/Engine/Source/Developer/ScriptDisassembler/Public/ScriptDisassembler.h b/Engine/Source/Developer/ScriptDisassembler/Public/ScriptDisassembler.h
new file mode 100644
--- /dev/null
+++ b/Engine/Source/Developer/ScriptDisassembler/Public/ScriptDisassembler.h
@@ -0,0 +1,86 @@
+// Copyright 1998-2019 Epic Games, Inc. All Rights Reserved.
+
+/*=============================================================================
+	ScriptDisassembler.h: Disassembler for Kismet bytecode.
+=============================================================================*/
+
+#pragma once
+
+#include "CoreMinimal.h"
+#include "UObject/Script.h"
+
+/**
+ * Kismet bytecode disassembler; Can be used to create a human readable version
+ * of Kismet bytecode for a specified structure or class.
+ */
+class FKismetBytecodeDisassembler
+{
+private:
+	TArray<uint8> Script;
+	FString Indents;
+	FOutputDevice& Ar;
+public:
+	/**
+	 * Construct a disassembler that will output to the specified archive.
+	 *
+	 * @param	InAr	The archive to emit disassembled bytecode to.
+	 */
+	SCRIPTDISASSEMBLER_API FKismetBytecodeDisassembler(FOutputDevice& InAr);
+
+	/**
+	 * Disassemble all of the script code in a single structure.
+	 *
+	 * @param [in,out]	Source	The structure to disassemble.
+	 */
+	SCRIPTDISASSEMBLER_API void DisassembleStructure(UFunction* Source);
+
+	/**
+	 * Disassemble all functions in any classes that have matching names.
+	 *
+	 * @param	InAr	The archive to emit disassembled bytecode to.
+	 * @param	ClassnameSubstring	A class must contain this substring to be disassembled.
+	 */
+	SCRIPTDISASSEMBLER_API static void DisassembleAllFunctionsInClasses(FOutputDevice& Ar, const FString& ClassnameSubstring);
+private:
+
+	// Reading functions
+	int32 ReadINT(int32& ScriptIndex);
+	uint64 ReadQWORD(int32& ScriptIndex);
+	uint8 ReadBYTE(int32& ScriptIndex);
+	FString ReadName(int32& ScriptIndex);
+	uint16 ReadWORD(int32& ScriptIndex);
+	float ReadFLOAT(int32& ScriptIndex);
+	CodeSkipSizeType ReadSkipCount(int32& ScriptIndex);
+	FString ReadString(int32& ScriptIndex);
+	FString ReadString8(int32& ScriptIndex);
+	FString ReadString16(int32& ScriptIndex);
+
+	EExprToken SerializeExpr(int32& ScriptIndex);
+	void ProcessCastByte(int32 CastType, int32& ScriptIndex);
+	void ProcessCommon(int32& ScriptIndex, EExprToken Opcode);
+
+	void InitTables();
+
+	template<typename T>
+	void Skip(int32& ScriptIndex)
+	{
+		ScriptIndex += sizeof(T);
+	}
+
+	void AddIndent()
+	{
+		Indents += TEXT("  ");
+	}
+
+	void DropIndent()
+	{
+		// Blah, this is awful
+		Indents = Indents.Left(Indents.Len() - 2);
+	}
+
+	template <typename T>
+	T* ReadPointer(int32& ScriptIndex)
+	{
+		return (T*)ReadQWORD(ScriptIndex);
+	}
+};
diff --git a/Engine/Source/Developer/ScriptDisassembler/ScriptDisassembler.Build.cs b/Engine/Source/Developer/ScriptDisassembler/ScriptDisassembler.Build.cs
new file mode 100644
--- /dev/null
+++ b/Engine/Source/Developer/ScriptDisassembler/ScriptDisassembler.Build.cs
@@ -0,0 +1,16 @@
+// Copyright 1998-2019 Epic Games, Inc. All Rights Reserved.
+
+namespace UnrealBuildTool.Rules
+{
+	public class ScriptDisassembler : ModuleRules
+	{
+		public ScriptDisassembler(ReadOnlyTargetRules Target) : base(Target)
+        {
+			PublicDependencyModuleNames.AddRange(new string[]
+			{
+				"Core",
+				"CoreUObject"
+			});
+        }
+	}
+}
diff --git a/Engine/Source/Editor/KismetCompiler/KismetCompiler.Build.cs b/Engine/Source/Editor/KismetCompiler/KismetCompiler.Build.cs
--- a/Engine/Source/Editor/KismetCompiler/KismetCompiler.Build.cs
+++ b/Engine/Source/Editor/KismetCompiler/KismetCompiler.Build.cs
@@ -19,6 +19,7 @@ public class KismetCompiler : ModuleRules
 				"AnimGraph",
                 "MessageLog",
 				"Kismet",
+				"ScriptDisassembler",
 			}
 			);
 	}
diff --git a/Engine/Source/Editor/UnrealEd/Private/ScriptDisassembler.cpp b/Engine/Source/Editor/UnrealEd/Private/ScriptDisassembler.cpp
deleted file mode 100644
diff --git a/Engine/Source/Editor/UnrealEd/Public/ScriptDisassembler.h b/Engine/Source/Editor/UnrealEd/Public/ScriptDisassembler.h
deleted file mode 100644
diff --git a/Engine/Source/Editor/UnrealEd/UnrealEd.Build.cs b/Engine/Source/Editor/UnrealEd/UnrealEd.Build.cs
--- a/Engine/Source/Editor/UnrealEd/UnrealEd.Build.cs
+++ b/Engine/Source/Editor/UnrealEd/UnrealEd.Build.cs
@@ -172,8 +172,9 @@ public class UnrealEd : ModuleRules
 				"TimeManagement",
                 "LandscapeEditorUtilities",
                 "DerivedDataCache",
+				"ScriptDisassembler",
 				"ToolMenus",
-            }
+			}
 		);
 
 		DynamicallyLoadedModuleNames.AddRange(
