Copy up from DS to Main

diff --git a/Engine/Source/Developer/ScriptDisassembler/Private/ScriptDisassemblerModule.cpp b/Engine/Source/Developer/ScriptDisassembler/Private/ScriptDisassemblerModule.cpp
deleted file mode 100644
diff --git a/Engine/Source/Developer/ScriptDisassembler/ScriptDisassembler.Build.cs b/Engine/Source/Developer/ScriptDisassembler/ScriptDisassembler.Build.cs
--- a/Engine/Source/Developer/ScriptDisassembler/ScriptDisassembler.Build.cs
+++ b/Engine/Source/Developer/ScriptDisassembler/ScriptDisassembler.Build.cs
@@ -6,6 +6,8 @@ namespace UnrealBuildTool.Rules
 	{
 		public ScriptDisassembler(ReadOnlyTargetRules Target) : base(Target)
         {
+			bRequiresImplementModule = false;
+
 			PublicDependencyModuleNames.AddRange(new string[]
 			{
 				"Core",
