Back out changelist 10320278

diff --git a/Engine/Source/Developer/ScriptDisassembler/Private/ScriptDisassemblerModule.cpp b/Engine/Source/Developer/ScriptDisassembler/Private/ScriptDisassemblerModule.cpp
new file mode 100644
--- /dev/null
+++ b/Engine/Source/Developer/ScriptDisassembler/Private/ScriptDisassemblerModule.cpp
@@ -0,0 +1,5 @@
+// Copyright 1998-2019 Epic Games, Inc. All Rights Reserved.
+
+#include "Modules/ModuleManager.h"
+
+IMPLEMENT_MODULE( FDefaultModuleImpl, ScriptDisassembler );
diff --git a/Engine/Source/Developer/ScriptDisassembler/ScriptDisassembler.Build.cs b/Engine/Source/Developer/ScriptDisassembler/ScriptDisassembler.Build.cs
--- a/Engine/Source/Developer/ScriptDisassembler/ScriptDisassembler.Build.cs
+++ b/Engine/Source/Developer/ScriptDisassembler/ScriptDisassembler.Build.cs
@@ -6,8 +6,6 @@ namespace UnrealBuildTool.Rules
 	{
 		public ScriptDisassembler(ReadOnlyTargetRules Target) : base(Target)
         {
-			bRequiresImplementModule = false;
-
 			PublicDependencyModuleNames.AddRange(new string[]
 			{
 				"Core",
