Fixed crash in the script disassembler.

We were still assuming that rotators and transforms were using floats, which is no longer the case with LWC.

diff --git a/Engine/Source/Developer/ScriptDisassembler/Private/ScriptDisassembler.cpp b/Engine/Source/Developer/ScriptDisassembler/Private/ScriptDisassembler.cpp
--- a/Engine/Source/Developer/ScriptDisassembler/Private/ScriptDisassembler.cpp
+++ b/Engine/Source/Developer/ScriptDisassembler/Private/ScriptDisassembler.cpp
@@ -156,6 +156,35 @@ FVector FKismetBytecodeDisassembler::ReadFVECTOR(int32& ScriptIndex)
 	return Vec;
 }
 
+FRotator FKismetBytecodeDisassembler::ReadFROTATOR(int32& ScriptIndex)
+{
+	FRotator Rotator;
+	Rotator.Pitch = ReadDOUBLE(ScriptIndex);
+	Rotator.Yaw = ReadDOUBLE(ScriptIndex);
+	Rotator.Roll = ReadDOUBLE(ScriptIndex);
+	return Rotator;
+}
+
+FQuat FKismetBytecodeDisassembler::ReadFQUAT(int32& ScriptIndex)
+{
+	FQuat Quat;
+	Quat.X = ReadDOUBLE(ScriptIndex);
+	Quat.Y = ReadDOUBLE(ScriptIndex);
+	Quat.Z = ReadDOUBLE(ScriptIndex);
+	Quat.W = ReadDOUBLE(ScriptIndex);
+	return Quat;
+}
+
+FTransform FKismetBytecodeDisassembler::ReadFTRANSFORM(int32& ScriptIndex)
+{
+	FTransform Transform;
+	FQuat TmpRotation = ReadFQUAT(ScriptIndex);
+	FVector TmpTranslation = ReadFVECTOR(ScriptIndex);
+	FVector TmpScale = ReadFVECTOR(ScriptIndex);
+	Transform.SetComponents(TmpRotation, TmpTranslation, TmpScale);
+	return Transform;
+}
+
 CodeSkipSizeType FKismetBytecodeDisassembler::ReadSkipCount(int32& ScriptIndex)
 {
 #if SCRIPT_LIMIT_BYTECODE_TO_64KB
@@ -895,11 +924,8 @@ void FKismetBytecodeDisassembler::ProcessCommon(int32& ScriptIndex, EExprToken O
 		}
 	case EX_RotationConst:
 		{
-			float Pitch = ReadFLOAT(ScriptIndex);
-			float Yaw = ReadFLOAT(ScriptIndex);
-			float Roll = ReadFLOAT(ScriptIndex);
-
-			Ar.Logf(TEXT("%s $%X: literal rotation (%f,%f,%f)"), *Indents, (int32)Opcode, Pitch, Yaw, Roll);
+			const FRotator Rotator = ReadFROTATOR(ScriptIndex);
+			Ar.Logf(TEXT("%s $%X: literal rotation (%f,%f,%f)"), *Indents, (int32)Opcode, Rotator.Pitch, Rotator.Yaw, Rotator.Roll);
 			break;
 		}
 	case EX_VectorConst:
@@ -916,15 +942,16 @@ void FKismetBytecodeDisassembler::ProcessCommon(int32& ScriptIndex, EExprToken O
 		}
 	case EX_TransformConst:
 		{
-			float RotX = ReadFLOAT(ScriptIndex);
-			float RotY = ReadFLOAT(ScriptIndex);
-			float RotZ = ReadFLOAT(ScriptIndex);
-			float RotW = ReadFLOAT(ScriptIndex);
-
-			FVector Trans = ReadFVECTOR(ScriptIndex);
-			FVector Scale = ReadFVECTOR(ScriptIndex);
-
-			Ar.Logf(TEXT("%s $%X: literal transform R(%f,%f,%f,%f) T(%f,%f,%f) S(%f,%f,%f)"), *Indents, (int32)Opcode, Trans.X, Trans.Y, Trans.Z, RotX, RotY, RotZ, RotW, Scale.X, Scale.Y, Scale.Z);
+			const FTransform Transform = ReadFTRANSFORM(ScriptIndex);
+			const FQuat Rotation = Transform.GetRotation();
+			const FVector Translation = Transform.GetTranslation();
+			const FVector Scale = Transform.GetScale3D();
+			Ar.Logf(TEXT("%s $%X: literal transform R(%f,%f,%f,%f) T(%f,%f,%f) S(%f,%f,%f)"),
+				    *Indents,
+					(int32)Opcode, 
+					Rotation.X, Rotation.Y, Rotation.Z, Rotation.W,
+					Translation.X, Translation.Y, Translation.Z,
+					Scale.X, Scale.Y, Scale.Z);
 			break;
 		}
 	case EX_StructConst:
diff --git a/Engine/Source/Developer/ScriptDisassembler/Public/ScriptDisassembler.h b/Engine/Source/Developer/ScriptDisassembler/Public/ScriptDisassembler.h
--- a/Engine/Source/Developer/ScriptDisassembler/Public/ScriptDisassembler.h
+++ b/Engine/Source/Developer/ScriptDisassembler/Public/ScriptDisassembler.h
@@ -52,6 +52,9 @@ private:
 	float ReadFLOAT(int32& ScriptIndex);
 	double ReadDOUBLE(int32& ScriptIndex);
 	FVector ReadFVECTOR(int32& ScriptIndex);
+	FRotator ReadFROTATOR(int32& ScriptIndex);
+	FQuat ReadFQUAT(int32& ScriptIndex);
+	FTransform ReadFTRANSFORM(int32& ScriptIndex);
 	CodeSkipSizeType ReadSkipCount(int32& ScriptIndex);
 	FString ReadString(int32& ScriptIndex);
 	FString ReadString8(int32& ScriptIndex);
