Handle CPF_RequiredParm in bytecode disassembler

diff --git a/Engine/Source/Developer/ScriptDisassembler/Private/ScriptDisassembler.cpp b/Engine/Source/Developer/ScriptDisassembler/Private/ScriptDisassembler.cpp
--- a/Engine/Source/Developer/ScriptDisassembler/Private/ScriptDisassembler.cpp
+++ b/Engine/Source/Developer/ScriptDisassembler/Private/ScriptDisassembler.cpp
@@ -287,6 +287,11 @@ void FKismetBytecodeDisassembler::ProcessCommon(int32& ScriptIndex, EExprToken O
 					ParameterType += TEXT("Return,");
 				}
 
+				if (PropertyPtr->HasAnyPropertyFlags(CPF_RequiredParm))
+				{
+					ParameterType += TEXT("Required,");
+				}
+
 				if (PropertyPtr->HasAnyPropertyFlags(CPF_ReferenceParm))
 				{
 					ParameterType += TEXT("Reference,");
