[Backout] - CL51791046. Just ScriptDisassembler changes. This change contained native serialization changes that cannot be backed out
#rnx
Original CL Desc
-----------------------------------------------------------------
Localization: Add a "Dev Notes" field to localized text so developers can provide additional context for translators

This is exposed as an extra text field in the editor and with new macros in C++ (NOTELOCTEXT, NSNOTELOCTEXT)

PR #11531

diff --git a/Engine/Source/Developer/ScriptDisassembler/Private/ScriptDisassembler.cpp b/Engine/Source/Developer/ScriptDisassembler/Private/ScriptDisassembler.cpp
--- a/Engine/Source/Developer/ScriptDisassembler/Private/ScriptDisassembler.cpp
+++ b/Engine/Source/Developer/ScriptDisassembler/Private/ScriptDisassembler.cpp
@@ -7,7 +7,6 @@
 #include "ScriptDisassembler.h"
 #include "UObject/Object.h"
 #include "UObject/Class.h"
-#include "UObject/Package.h"
 #include "UObject/UObjectHash.h"
 #include "UObject/UObjectIterator.h"
 #include "UObject/UnrealType.h"
@@ -27,7 +26,6 @@ FKismetBytecodeDisassembler::FKismetBytecodeDisassembler(FOutputDevice& InAr)
 // Disassemble all of the script code in a single structure.
 void FKismetBytecodeDisassembler::DisassembleStructure(UFunction* Source)
 {
-	bIsCooked = Source->GetPackage()->HasAnyPackageFlags(PKG_Cooked);
 	Script.Empty();
 	Script.Append(Source->Script);
 
@@ -862,17 +860,7 @@ void FKismetBytecodeDisassembler::ProcessCommon(int32& ScriptIndex, EExprToken O
 					const FString SourceString = ReadString(ScriptIndex);
 					const FString KeyString = ReadString(ScriptIndex);
 					const FString Namespace = ReadString(ScriptIndex);
-#if WITH_EDITORONLY_DATA
-					if (!bIsCooked)
-					{
-						const FString DevNotes = ReadString(ScriptIndex);
-						Ar.Logf(TEXT("%s $%X: literal text - localized text { namespace: \"%s\", key: \"%s\", source: \"%s\", devnotes: \"%s\" }"), *Indents, (int32)Opcode, *Namespace, *KeyString, *SourceString, *DevNotes);
-					}
-					else
-#endif
-					{
-						Ar.Logf(TEXT("%s $%X: literal text - localized text { namespace: \"%s\", key: \"%s\", source: \"%s\" }"), *Indents, (int32)Opcode, *Namespace, *KeyString, *SourceString);
-					}
+					Ar.Logf(TEXT("%s $%X: literal text - localized text { namespace: \"%s\", key: \"%s\", source: \"%s\" }"), *Indents, (int32)Opcode, *Namespace, *KeyString, *SourceString);
 				}
 				break;
 
diff --git a/Engine/Source/Developer/ScriptDisassembler/Public/ScriptDisassembler.h b/Engine/Source/Developer/ScriptDisassembler/Public/ScriptDisassembler.h
--- a/Engine/Source/Developer/ScriptDisassembler/Public/ScriptDisassembler.h
+++ b/Engine/Source/Developer/ScriptDisassembler/Public/ScriptDisassembler.h
@@ -28,7 +28,6 @@ private:
 	TArray<uint8> Script;
 	FString Indents;
 	FOutputDevice& Ar;
-	bool bIsCooked = false;
 public:
 	/**
 	 * Construct a disassembler that will output to the specified archive.
