Localization: Fix gathering of devnotes on BP graph text literals

PR #11531

#was-backout 51802337
#rnx
#Synced-CL 51946586

diff --git a/Engine/Source/Developer/ScriptDisassembler/Private/ScriptDisassembler.cpp b/Engine/Source/Developer/ScriptDisassembler/Private/ScriptDisassembler.cpp
--- a/Engine/Source/Developer/ScriptDisassembler/Private/ScriptDisassembler.cpp
+++ b/Engine/Source/Developer/ScriptDisassembler/Private/ScriptDisassembler.cpp
@@ -856,11 +856,20 @@ void FKismetBytecodeDisassembler::ProcessCommon(int32& ScriptIndex, EExprToken O
 				break;
 
 			case EBlueprintTextLiteralType::LocalizedText:
+			case EBlueprintTextLiteralType::LocalizedTextWithNotes:
 				{
 					const FString SourceString = ReadString(ScriptIndex);
 					const FString KeyString = ReadString(ScriptIndex);
 					const FString Namespace = ReadString(ScriptIndex);
-					Ar.Logf(TEXT("%s $%X: literal text - localized text { namespace: \"%s\", key: \"%s\", source: \"%s\" }"), *Indents, (int32)Opcode, *Namespace, *KeyString, *SourceString);
+					if (TextLiteralType == EBlueprintTextLiteralType::LocalizedTextWithNotes)
+					{
+						const FString DevNotes = ReadString(ScriptIndex);
+						Ar.Logf(TEXT("%s $%X: literal text - localized text { namespace: \"%s\", key: \"%s\", source: \"%s\", devnotes: \"%s\" }"), *Indents, (int32)Opcode, *Namespace, *KeyString, *SourceString, *DevNotes);
+					}
+					else
+					{
+						Ar.Logf(TEXT("%s $%X: literal text - localized text { namespace: \"%s\", key: \"%s\", source: \"%s\" }"), *Indents, (int32)Opcode, *Namespace, *KeyString, *SourceString);
+					}
 				}
 				break;
 
diff --git a/Engine/Source/Editor/KismetCompiler/Private/KismetCompilerVMBackend.cpp b/Engine/Source/Editor/KismetCompiler/Private/KismetCompilerVMBackend.cpp
--- a/Engine/Source/Editor/KismetCompiler/Private/KismetCompilerVMBackend.cpp
+++ b/Engine/Source/Editor/KismetCompiler/Private/KismetCompilerVMBackend.cpp
@@ -725,11 +725,23 @@ public:
 						// If we're gathering editor-only text then we'll pick up the version with the package localization ID from the property/pin rather than the bytecode
 						const FString Namespace = TextNamespaceUtil::StripPackageNamespace(TextId.GetNamespace().ToString());
 						const FString Key = TextId.GetKey().ToString();
+#if WITH_EDITORONLY_DATA
+						const FString& DevNotes = Term->TextLiteral.GetDevNotes();
+						const bool bEmitDevNotes = !Writer.IsFilterEditorOnly() && !DevNotes.IsEmpty();
+#else
+						constexpr bool bEmitDevNotes = false;
+#endif
 
-						Writer << EBlueprintTextLiteralType::LocalizedText;
+						Writer << (bEmitDevNotes ? EBlueprintTextLiteralType::LocalizedTextWithNotes : EBlueprintTextLiteralType::LocalizedText);
 						EmitStringLiteral(*SourceString);
 						EmitStringLiteral(Key);
 						EmitStringLiteral(Namespace);
+#if WITH_EDITORONLY_DATA
+						if (bEmitDevNotes)
+						{
+							EmitStringLiteral(DevNotes);
+						}
+#endif
 					}
 					else
 					{
diff --git a/Engine/Source/Runtime/CoreUObject/Private/Serialization/PropertyLocalizationDataGathering.cpp b/Engine/Source/Runtime/CoreUObject/Private/Serialization/PropertyLocalizationDataGathering.cpp
--- a/Engine/Source/Runtime/CoreUObject/Private/Serialization/PropertyLocalizationDataGathering.cpp
+++ b/Engine/Source/Runtime/CoreUObject/Private/Serialization/PropertyLocalizationDataGathering.cpp
@@ -743,6 +743,7 @@ private:
 			break;
 
 		case EBlueprintTextLiteralType::LocalizedText:
+		case EBlueprintTextLiteralType::LocalizedTextWithNotes:
 			{
 				bIsParsingText = true;
 
@@ -754,10 +755,21 @@ private:
 
 				SerializeExpr(iCode, Ar);
 				const FString TextNamespace = MoveTemp(LastParsedString);
+				
+				FString DevNotes;
+				if (TextLiteralType == EBlueprintTextLiteralType::LocalizedTextWithNotes)
+				{
+					SerializeExpr(iCode, Ar);
+					DevNotes = MoveTemp(LastParsedString);
+				}
 
 				bIsParsingText = false;
 
+#if WITH_EDITORONLY_DATA
+				const FText TextInstance = FText::AsLocalizable_Advanced(TextNamespace, TextKey, MoveTemp(SourceString), DevNotes);
+#else
 				const FText TextInstance = FText::AsLocalizable_Advanced(TextNamespace, TextKey, MoveTemp(SourceString));
+#endif
 				if (!PropertyLocalizationDataGatherer.IsDefaultTextInstance(TextInstance))
 				{
 					PropertyLocalizationDataGatherer.GatherTextInstance(TextInstance, FString::Printf(TEXT("%s [Script Bytecode]"), SourceDescription), bTreatAsEditorOnlyData);
diff --git a/Engine/Source/Runtime/CoreUObject/Private/UObject/ScriptCore.cpp b/Engine/Source/Runtime/CoreUObject/Private/UObject/ScriptCore.cpp
--- a/Engine/Source/Runtime/CoreUObject/Private/UObject/ScriptCore.cpp
+++ b/Engine/Source/Runtime/CoreUObject/Private/UObject/ScriptCore.cpp
@@ -3472,6 +3472,7 @@ DEFINE_FUNCTION(UObject::execTextConst)
 		break;
 
 	case EBlueprintTextLiteralType::LocalizedText:
+	case EBlueprintTextLiteralType::LocalizedTextWithNotes:
 		{
 			FString SourceString;
 			Stack.Step(Stack.Object, &SourceString);
@@ -3482,7 +3483,16 @@ DEFINE_FUNCTION(UObject::execTextConst)
 			FString Namespace;
 			Stack.Step(Stack.Object, &Namespace);
 
+			FString DevNotes;
+			if (TextLiteralType == EBlueprintTextLiteralType::LocalizedTextWithNotes)
+			{
+				Stack.Step(Stack.Object, &DevNotes);
+			}
+#if WITH_EDITORONLY_DATA
+			*(FText*)RESULT_PARAM = FText::AsLocalizable_Advanced(Namespace, KeyString, MoveTemp(SourceString), DevNotes);
+#else
 			*(FText*)RESULT_PARAM = FText::AsLocalizable_Advanced(Namespace, KeyString, MoveTemp(SourceString));
+#endif
 		}
 		break;
 
diff --git a/Engine/Source/Runtime/CoreUObject/Public/UObject/Script.h b/Engine/Source/Runtime/CoreUObject/Public/UObject/Script.h
--- a/Engine/Source/Runtime/CoreUObject/Public/UObject/Script.h
+++ b/Engine/Source/Runtime/CoreUObject/Public/UObject/Script.h
@@ -337,6 +337,8 @@ enum class EBlueprintTextLiteralType : uint8
 	Empty,
 	/** Text is localized. The bytecode will contain three strings - source, key, and namespace - and should be loaded via FInternationalization */
 	LocalizedText,
+	/** Text is localized. The bytecode will contain four strings - source, key, namespace, and devnotes - and should be loaded via FInternationalization */
+	LocalizedTextWithNotes,
 	/** Text is culture invariant. The bytecode will contain one string, and you should use FText::AsCultureInvariant to initialize the FText instance. */
 	InvariantText,
 	/** Text is a literal FString. The bytecode will contain one string, and you should use FText::FromString to initialize the FText instance. */
diff --git a/Engine/Source/Runtime/CoreUObject/Public/UObject/ScriptSerialization.inl b/Engine/Source/Runtime/CoreUObject/Public/UObject/ScriptSerialization.inl
--- a/Engine/Source/Runtime/CoreUObject/Public/UObject/ScriptSerialization.inl
+++ b/Engine/Source/Runtime/CoreUObject/Public/UObject/ScriptSerialization.inl
@@ -78,6 +78,8 @@
 		{ \
 		case EBlueprintTextLiteralType::Empty: \
 			break; \
+		case EBlueprintTextLiteralType::LocalizedTextWithNotes: \
+			SerializeExpr( iCode, Ar ); \
 		case EBlueprintTextLiteralType::LocalizedText: \
 			SerializeExpr( iCode, Ar );	\
 			SerializeExpr( iCode, Ar ); \
